[RewindDebugger] applied review feedback from previous change (removed redundant semi-colons)

#Synced-CL 49932076

diff --git a/Engine/Source/Runtime/RewindDebuggerRuntimeInterface/Public/RewindDebuggerRuntimeInterface/IRewindDebuggerRuntimeExtension.h b/Engine/Source/Runtime/RewindDebuggerRuntimeInterface/Public/RewindDebuggerRuntimeInterface/IRewindDebuggerRuntimeExtension.h
--- a/Engine/Source/Runtime/RewindDebuggerRuntimeInterface/Public/RewindDebuggerRuntimeInterface/IRewindDebuggerRuntimeExtension.h
+++ b/Engine/Source/Runtime/RewindDebuggerRuntimeInterface/Public/RewindDebuggerRuntimeInterface/IRewindDebuggerRuntimeExtension.h
@@ -45,17 +45,17 @@ public:
 	/** called when recording has started */
 	virtual void RecordingStarted()
 	{
-	};
+	}
 
 	/** called when recording has ended */
 	virtual void RecordingStopped()
 	{
-	};
+	}
 
 	/** called when recording is unloaded or a new one starts */
 	virtual void Clear()
 	{
-	};
+	}
 
 	/**
 	 * This method allows extensions to register their own message type handlers to the message system and send response using
@@ -65,7 +65,7 @@ public:
 	 */
 	virtual void RegisterMessageHandlers(const TSharedPtr<UE::TraceBasedDebuggers::FRemoteSessionsManager>& InRemoteSessionsManager, FMessageEndpointBuilder& InMessageEndpointBuilder)
 	{
-	};
+	}
 
 	/**
 	 * This method allows extensions to register their own message types in the remote sessions manager and subscribe in the message endpoint.
@@ -74,7 +74,7 @@ public:
 	 */
 	virtual void RegisterMessageTypes(const TSharedPtr<UE::TraceBasedDebuggers::FRemoteSessionsManager>& InRemoteSessionsManager, FMessageEndpoint& InMessageEndpoint)
 	{
-	};
+	}
 };
 
 #undef UE_API
